Compilation fixes for case-sensitive filesystems.

diff --git a/Engine/Source/Editor/GeometryCacheEd/Classes/GeometryCacheThumbnailScene.h b/Engine/Source/Editor/GeometryCacheEd/Classes/GeometryCacheThumbnailScene.h
--- a/Engine/Source/Editor/GeometryCacheEd/Classes/GeometryCacheThumbnailScene.h
+++ b/Engine/Source/Editor/GeometryCacheEd/Classes/GeometryCacheThumbnailScene.h
@@ -2,7 +2,7 @@
 
 #pragma once
 
-#include "Thumbnailhelpers.h"
+#include "ThumbnailHelpers.h"
 
 class UGeometryCache;
 class AGeometryCacheActor;
@@ -23,4 +23,4 @@ protected:
 private:
 	/** The static mesh actor used to display all static mesh thumbnails */
 	AGeometryCacheActor* PreviewActor;
-};
\ No newline at end of file
+};
diff --git a/Engine/Source/Editor/GeometryCacheEd/Private/GeometryCacheThumbnailRenderer.cpp b/Engine/Source/Editor/GeometryCacheEd/Private/GeometryCacheThumbnailRenderer.cpp
--- a/Engine/Source/Editor/GeometryCacheEd/Private/GeometryCacheThumbnailRenderer.cpp
+++ b/Engine/Source/Editor/GeometryCacheEd/Private/GeometryCacheThumbnailRenderer.cpp
@@ -3,9 +3,9 @@
 #include "GeometryCacheEdModulePublicPCH.h"
 #include "GeometryCacheThumbnailRenderer.h"
 #include "GeometryCacheThumbnailScene.h"
-#include "Thumbnailhelpers.h"
+#include "ThumbnailHelpers.h"
 #include "EngineModule.h"
-#include "Sceneview.h"
+#include "SceneView.h"
 #include "RendererInterface.h"
 #include "GeometryCache.h"
 
diff --git a/Engine/Source/Runtime/GeometryCache/Private/GeometryCache.cpp b/Engine/Source/Runtime/GeometryCache/Private/GeometryCache.cpp
--- a/Engine/Source/Runtime/GeometryCache/Private/GeometryCache.cpp
+++ b/Engine/Source/Runtime/GeometryCache/Private/GeometryCache.cpp
@@ -3,7 +3,7 @@
 #include "GeometryCacheModulePrivatePCH.h"
 #include "GeometryCacheTrackTransformAnimation.h"
 #include "GeometryCacheTrackFlipbookAnimation.h"
-#include "EditorFrameWork/AssetImportData.h"
+#include "EditorFramework/AssetImportData.h"
 
 UGeometryCache::UGeometryCache(const FObjectInitializer& ObjectInitializer /*= FObjectInitializer::Get()*/) : UObject(ObjectInitializer)
 {
