Mac Non-unity build errors

diff --git a/Engine/Source/Editor/GeometryCacheEd/Public/GeometryCacheEdModulePublicPCH.h b/Engine/Source/Editor/GeometryCacheEd/Public/GeometryCacheEdModulePublicPCH.h
--- a/Engine/Source/Editor/GeometryCacheEd/Public/GeometryCacheEdModulePublicPCH.h
+++ b/Engine/Source/Editor/GeometryCacheEd/Public/GeometryCacheEdModulePublicPCH.h
@@ -1,7 +1,5 @@
 // Copyright 1998-2015 Epic Games, Inc. All Rights Reserved.
 
-//#include "ActorFactoryGeometryCache.h"
-//#include "AssetTypeActions_GeometryCache.h"
-//#include "GeometryCacheAssetBroker.h"
-//#include "GeometryCacheThumbnailRenderer.h"
-//#include "GeometryCacheThumbnailScene.h"
\ No newline at end of file
+#include "Core.h"
+#include "CoreUObject.h"
+#include "Engine.h"
\ No newline at end of file
