More Mac non-unity errors

diff --git a/Engine/Source/Editor/GeometryCacheEd/Classes/AssetTypeActions_GeometryCache.h b/Engine/Source/Editor/GeometryCacheEd/Classes/AssetTypeActions_GeometryCache.h
--- a/Engine/Source/Editor/GeometryCacheEd/Classes/AssetTypeActions_GeometryCache.h
+++ b/Engine/Source/Editor/GeometryCacheEd/Classes/AssetTypeActions_GeometryCache.h
@@ -3,6 +3,7 @@
 
 #include "GeometryCache.h"
 #include "AssetTypeActions_Base.h"
+#include "SceneThumbnailInfo.h"
 
 /** Type-actions class for GeometryCache assets*/
 class FAssetTypeActions_GeometryCache : public FAssetTypeActions_Base
