Mac unity build errors

diff --git a/Engine/Source/Editor/GeometryCacheEd/Classes/GeometryCacheThumbnailRenderer.h b/Engine/Source/Editor/GeometryCacheEd/Classes/GeometryCacheThumbnailRenderer.h
--- a/Engine/Source/Editor/GeometryCacheEd/Classes/GeometryCacheThumbnailRenderer.h
+++ b/Engine/Source/Editor/GeometryCacheEd/Classes/GeometryCacheThumbnailRenderer.h
@@ -6,8 +6,13 @@
 */
 
 #pragma once
+
+#include "GeometryCacheEdModulePublicPCH.h"
+#include "Classes/ThumbnailRendering/ThumbnailRenderer.h"
+#include "Classes/ThumbnailRendering/DefaultSizedThumbnailRenderer.h"
 #include "GeometryCacheThumbnailRenderer.generated.h"
 
+
 class FGeometryCacheThumbnailScene;
 
 UCLASS(config = Editor, MinimalAPI)
diff --git a/Engine/Source/Editor/GeometryCacheEd/Classes/GeometryCacheThumbnailScene.h b/Engine/Source/Editor/GeometryCacheEd/Classes/GeometryCacheThumbnailScene.h
--- a/Engine/Source/Editor/GeometryCacheEd/Classes/GeometryCacheThumbnailScene.h
+++ b/Engine/Source/Editor/GeometryCacheEd/Classes/GeometryCacheThumbnailScene.h
@@ -2,6 +2,7 @@
 
 #pragma once
 
+#include "GeometryCacheEdModulePublicPCH.h"
 #include "ThumbnailHelpers.h"
 
 class UGeometryCache;
diff --git a/Engine/Source/Editor/GeometryCacheEd/Public/GeometryCacheEdModulePublicPCH.h b/Engine/Source/Editor/GeometryCacheEd/Public/GeometryCacheEdModulePublicPCH.h
--- a/Engine/Source/Editor/GeometryCacheEd/Public/GeometryCacheEdModulePublicPCH.h
+++ b/Engine/Source/Editor/GeometryCacheEd/Public/GeometryCacheEdModulePublicPCH.h
@@ -2,4 +2,5 @@
 
 #include "Core.h"
 #include "CoreUObject.h"
-#include "Engine.h"
\ No newline at end of file
+#include "Engine.h"
+#include "UnrealEd.h"
\ No newline at end of file
