Fix for CIS warnings, and fix for GeometryCache Thumbnail renderer (missed out code in merge)

diff --git a/Engine/Config/BaseEditor.ini b/Engine/Config/BaseEditor.ini
--- a/Engine/Config/BaseEditor.ini
+++ b/Engine/Config/BaseEditor.ini
@@ -40,7 +40,6 @@ MaxNumberOfLogLines=2000
 +RenderableThumbnailTypes=(ClassNeedingThumbnailName="/Script/Engine.SoundWave",RendererClassName="/Script/UnrealEd.SoundWaveThumbnailRenderer")
 +RenderableThumbnailTypes=(ClassNeedingThumbnailName="/Script/Engine.SlateBrushAsset",RendererClassName="/Script/UnrealEd.SlateBrushThumbnailRenderer")
 +RenderableThumbnailTypes=(ClassNeedingThumbnailName="/Script/Engine.SubsurfaceProfile",RendererClassName="/Script/UnrealEd.SubsurfaceProfileRenderer")
-+RenderableThumbnailTypes=(ClassNeedingThumbnailName="/Script/GeometryCache.GeometryCache",RendererClassName="/Script/GeometryCacheEd.GeometryCacheThumbnailRenderer")
 
 [/Script/UnrealEd.DefaultSizedThumbnailRenderer]
 DefaultSizeX=512
diff --git a/Engine/Source/Editor/GeometryCacheEd/Private/GeometryCacheEdModule.cpp b/Engine/Source/Editor/GeometryCacheEd/Private/GeometryCacheEdModule.cpp
--- a/Engine/Source/Editor/GeometryCacheEd/Private/GeometryCacheEdModule.cpp
+++ b/Engine/Source/Editor/GeometryCacheEd/Private/GeometryCacheEdModule.cpp
@@ -7,6 +7,7 @@
 #include "AssetTypeActions_GeometryCache.h"
 #include "GeometryCacheAssetBroker.h"
 #include "GeometryCacheComponent.h"
+#include "GeometryCacheThumbnailRenderer.h"
 
 IMPLEMENT_MODULE(FGeometryCacheEdModule, GeometryCacheEd)
 
@@ -20,6 +21,8 @@ void FGeometryCacheEdModule::StartupModule()
 
 	AssetBroker = new FGeometryCacheAssetBroker();
 	FComponentAssetBrokerage::RegisterBroker(MakeShareable(AssetBroker), UGeometryCacheComponent::StaticClass(), true, true);
+
+	UThumbnailManager::Get().RegisterCustomRenderer(UGeometryCache::StaticClass(), UGeometryCacheThumbnailRenderer::StaticClass());
 }
 
 void FGeometryCacheEdModule::ShutdownModule()
diff --git a/Engine/Source/Editor/GeometryCacheEd/Private/GeometryCacheThumbnailScene.cpp b/Engine/Source/Editor/GeometryCacheEd/Private/GeometryCacheThumbnailScene.cpp
--- a/Engine/Source/Editor/GeometryCacheEd/Private/GeometryCacheThumbnailScene.cpp
+++ b/Engine/Source/Editor/GeometryCacheEd/Private/GeometryCacheThumbnailScene.cpp
@@ -64,5 +64,10 @@ void FGeometryCacheThumbnailScene::GetViewMatrixParameters(const float InFOVDegr
 	{
 		ThumbnailInfo = USceneThumbnailInfo::StaticClass()->GetDefaultObject<USceneThumbnailInfo>();
 	}
+
+	OutOrigin = FVector(0, 0, -BoundsZOffset);
+	OutOrbitPitch = ThumbnailInfo->OrbitPitch;
+	OutOrbitYaw = ThumbnailInfo->OrbitYaw;
+	OutOrbitZoom = TargetDistance + ThumbnailInfo->OrbitZoom;
 }
 
