Geom Cache Thumbnail safety

diff --git a/Engine/Source/Editor/GeometryCacheEd/Private/GeometryCacheEdModule.cpp b/Engine/Source/Editor/GeometryCacheEd/Private/GeometryCacheEdModule.cpp
--- a/Engine/Source/Editor/GeometryCacheEd/Private/GeometryCacheEdModule.cpp
+++ b/Engine/Source/Editor/GeometryCacheEd/Private/GeometryCacheEdModule.cpp
@@ -33,5 +33,6 @@ void FGeometryCacheEdModule::ShutdownModule()
 		IAssetTools& AssetTools = AssetToolsModule.Get();
 		AssetTools.UnregisterAssetTypeActions(AssetAction->AsShared());
 		FComponentAssetBrokerage::UnregisterBroker(MakeShareable(AssetBroker));
+		UThumbnailManager::Get().UnregisterCustomRenderer(UGeometryCache::StaticClass());
 	}
 }
\ No newline at end of file
